Initial HairStrandsMutable plugin submit.

Mutable update now will instanciate/remove groom components. Groom components will be placed as siblings of the UCustomizableObjectInstanceUsage.

- Added new Groom Constant node.

#rnx

diff --git a/Engine/Plugins/Experimental/HairStrandsMutable/HairStrandsMutable.uplugin b/Engine/Plugins/Experimental/HairStrandsMutable/HairStrandsMutable.uplugin
new file mode 100644
--- /dev/null
+++ b/Engine/Plugins/Experimental/HairStrandsMutable/HairStrandsMutable.uplugin
@@ -0,0 +1,46 @@
+{
+	"FileVersion": 3,
+	"Version": 1,
+	"VersionName": "0.1",
+	"FriendlyName": "Groom Extensions For Mutable",
+	"Description": "Adds Mutable functionality to work with Grooms from the HairStrands plugin",
+	"Category": "Animation",
+	"CreatedBy": "Epic Games, Inc.",
+	"CreatedByURL": "https://epicgames.com",
+	"DocsURL": "",
+	"MarketplaceURL": "",
+	"SupportURL": "",
+	"CanContainContent": false,
+	"IsBetaVersion": true,
+	"IsExperimentalVersion": true,
+	"Installed": false,
+	"EnabledByDefault": false,
+	"Modules": [
+		{
+			"Name": "HairStrandsMutable",
+			"Type": "Runtime",
+			"LoadingPhase": "Default"
+		},
+		{
+			"Name": "HairStrandsMutableEditor",
+			"Type": "Editor",
+			"LoadingPhase": "Default"
+		}
+
+	],
+	"Plugins": [
+		{
+			"Name": "HairStrands",
+			"Enabled": true
+		},
+		{
+			"Name": "Mutable",
+			"Enabled": true
+		},
+		{
+			"Name": "StructUtils",
+			"Enabled": true	
+		}
+	]
+
+}
\ No newline at end of file
diff --git a/Engine/Plugins/Experimental/HairStrandsMutable/Source/HairStrandsMutable/HairStrandsMutable.Build.cs b/Engine/Plugins/Experimental/HairStrandsMutable/Source/HairStrandsMutable/HairStrandsMutable.Build.cs
new file mode 100644
--- /dev/null
+++ b/Engine/Plugins/Experimental/HairStrandsMutable/Source/HairStrandsMutable/HairStrandsMutable.Build.cs
@@ -0,0 +1,30 @@
+// Copyright Epic Games, Inc. All Rights Reserved.
+
+namespace UnrealBuildTool.Rules
+{
+	public class HairStrandsMutable : ModuleRules
+	{
+        public HairStrandsMutable(ReadOnlyTargetRules Target) : base(Target)
+		{
+			PrivateIncludePaths.Add(ModuleDirectory + "/Private");
+			PublicIncludePaths.Add(ModuleDirectory + "/Public");
+			
+			PublicDependencyModuleNames.AddRange(
+				new string[]
+				{
+					"Core",
+					"CoreUObject",
+					"Engine",
+					"CustomizableObject",
+					"HairStrandsCore",
+					"StructUtils",
+				});
+
+			PrivateDependencyModuleNames.AddRange(
+				new string[]
+				{
+					"Niagara"
+				});
+		}
+	}
+}
diff --git a/Engine/Plugins/Experimental/HairStrandsMutable/Source/HairStrandsMutable/Private/HairStrandsMutableExtension.cpp b/Engine/Plugins/Experimental/HairStrandsMutable/Source/HairStrandsMutable/Private/HairStrandsMutableExtension.cpp
new file mode 100644
--- /dev/null
+++ b/Engine/Plugins/Experimental/HairStrandsMutable/Source/HairStrandsMutable/Private/HairStrandsMutableExtension.cpp
@@ -0,0 +1,160 @@
+// Copyright Epic Games, Inc. All Rights Reserved.
+
+#include "HairStrandsMutableExtension.h"
+
+#include "Algo/AnyOf.h"
+#include "GroomComponent.h"
+#include "Components/SkeletalMeshComponent.h"
+#include "MuCO/CustomizableObject.h"
+#include "MuCO/CustomizableObjectInstance.h"
+#include "MuCO/CustomizableObjectInstanceUsage.h"
+#include "MuCO/CustomizableSkeletalComponent.h"
+#include "MuCO/CustomizableSkeletalMeshActor.h"
+#include "UObject/UObjectGlobals.h"
+
+const FName GroomComponentTag(TEXT("Mutable"));
+
+const FName UHairStrandsMutableExtension::GroomPinType(TEXT("Groom"));
+const FName UHairStrandsMutableExtension::GroomsBaseNodePinName(TEXT("Grooms"));
+const FText UHairStrandsMutableExtension::GroomNodeCategory(FText::FromString(TEXT("Grooms")));
+
+TArray<FCustomizableObjectPinType> UHairStrandsMutableExtension::GetPinTypes() const
+{
+	TArray<FCustomizableObjectPinType> Result;
+	
+	FCustomizableObjectPinType& GroomType = Result.AddDefaulted_GetRef();
+	GroomType.Name = GroomPinType;
+	GroomType.DisplayName = FText::FromString(TEXT("Groom"));
+	GroomType.Color = FLinearColor::Red;
+
+	return Result;
+}
+
+TArray<FObjectNodeInputPin> UHairStrandsMutableExtension::GetAdditionalObjectNodePins() const
+{
+	TArray<FObjectNodeInputPin> Result;
+
+	FObjectNodeInputPin& GroomInputPin = Result.AddDefaulted_GetRef();
+	GroomInputPin.PinType = GroomPinType;
+	GroomInputPin.PinName = GroomsBaseNodePinName;
+	GroomInputPin.DisplayName = FText::FromString(TEXT("Groom"));
+	GroomInputPin.bIsArray = true;
+
+	return Result;
+}
+
+FInstancedStruct UHairStrandsMutableExtension::GenerateExtensionInstanceData(
+	const TArray<FInputPinDataContainer>& InputPinData) const
+{
+	FInstancedStruct Result(FGroomInstanceData::StaticStruct());
+	FGroomInstanceData* InstanceData = Result.GetMutablePtr<FGroomInstanceData>();
+
+	for (const FInputPinDataContainer& Container : InputPinData)
+	{
+		if (Container.Pin.PinName == GroomsBaseNodePinName)
+		{
+			const FGroomPinData* Data = Container.Data.GetPtr<FGroomPinData>();
+			if (Data
+				&& Data->GroomAsset != nullptr)
+			{
+				InstanceData->Grooms.Add(*Data);
+			}
+		}
+	}
+
+	return Result;
+}
+
+
+void UHairStrandsMutableExtension::OnCustomizableObjectInstanceUsageUpdated(const UCustomizableObjectInstanceUsage* InstanceUsage) const
+{
+	UCustomizableObjectInstance* Instance = InstanceUsage->GetCustomizableObjectInstance();
+	if (!Instance)
+	{
+		return;
+	}
+
+	USkeletalMeshComponent* AttachedParent = InstanceUsage->GetAttachParent();
+	if (!AttachedParent)
+	{
+		return;
+	}
+	
+	const FInstancedStruct UntypedInstanceData = Instance->GetExtensionInstanceData(this);
+	const FGroomInstanceData* InstanceData = UntypedInstanceData.GetPtr<FGroomInstanceData>();
+	if (!InstanceData)
+	{
+		return;
+	}
+
+	USceneComponent* AttachParent = InstanceUsage->GetAttachParent();
+	
+	TArray<UGroomComponent*> UpdateGroomComponents; // Grooms that must be present at the end of this update.
+
+	for (const FGroomPinData& GroomPinData : InstanceData->Grooms)
+	{
+		if (GroomPinData.ComponentIndex != InstanceUsage->GetComponentIndex())
+		{
+			continue;
+		}
+		
+		bool bAlreadyExists = false;
+
+		for (USceneComponent* Component : AttachParent->GetAttachChildren())
+		{
+			UGroomComponent* GroomComponent = Cast<UGroomComponent>(Component);
+			if (!GroomComponent)
+			{
+				continue;
+			}
+
+			if (GroomComponent->GroomAsset == GroomPinData.GroomAsset &&
+				GroomComponent->GroomCache == GroomPinData.GroomCache &&
+				GroomComponent->BindingAsset == GroomPinData.BindingAsset &&
+				GroomComponent->PhysicsAsset == GroomPinData.PhysicsAsset &&
+				GroomComponent->AttachmentName == GroomPinData.AttachmentName &&
+				GroomComponent->OverrideMaterials == GroomPinData.OverrideMaterials)
+			{
+				UpdateGroomComponents.Add(GroomComponent);
+				bAlreadyExists = true;
+				break;
+			}
+		}
+
+		if (bAlreadyExists)
+		{
+			continue;
+		}
+		
+		UGroomComponent* GroomComponent = NewObject<UGroomComponent>(AttachParent);
+		GroomComponent->GroomAsset = GroomPinData.GroomAsset;
+		GroomComponent->GroomCache = GroomPinData.GroomCache;
+		GroomComponent->BindingAsset = GroomPinData.BindingAsset;
+		GroomComponent->PhysicsAsset = GroomPinData.PhysicsAsset;
+		GroomComponent->AttachmentName = GroomPinData.AttachmentName;
+		GroomComponent->OverrideMaterials = GroomPinData.OverrideMaterials;
+		
+		// Work around UE-158069
+		GroomComponent->CreationMethod = EComponentCreationMethod::Instance;
+
+		GroomComponent->ComponentTags.Add(GroomComponentTag);
+		
+		GroomComponent->AttachToComponent(AttachParent, FAttachmentTransformRules::KeepRelativeTransform);
+		GroomComponent->RegisterComponent();
+
+		UpdateGroomComponents.Add(GroomComponent);
+	}
+
+	// Destroy old Mutable grooms.
+	TArray<USceneComponent*> Children;
+	AttachParent->GetChildrenComponents(false, Children);
+	for (USceneComponent* Child : Children)
+	{
+		if (Child->IsA(UGroomComponent::StaticClass()) &&
+			Child->ComponentTags.Contains(GroomComponentTag) &&
+			!UpdateGroomComponents.Contains(Child))
+		{
+			Child->DestroyComponent();
+		}
+	}
+}
diff --git a/Engine/Plugins/Experimental/HairStrandsMutable/Source/HairStrandsMutable/Private/HairStrandsMutableModule.cpp b/Engine/Plugins/Experimental/HairStrandsMutable/Source/HairStrandsMutable/Private/HairStrandsMutableModule.cpp
new file mode 100644
--- /dev/null
+++ b/Engine/Plugins/Experimental/HairStrandsMutable/Source/HairStrandsMutable/Private/HairStrandsMutableModule.cpp
@@ -0,0 +1,29 @@
+// Copyright Epic Games, Inc. All Rights Reserved.
+
+#include "HairStrandsMutableModule.h"
+
+#include "MuCO/ICustomizableObjectModule.h"
+#include "HairStrandsMutableExtension.h"
+
+class FHairStrandsMutableModule : public IModuleInterface
+{
+	virtual void StartupModule() override
+	{
+		Extension = TWeakObjectPtr<const UHairStrandsMutableExtension>(GetDefault<UHairStrandsMutableExtension>());
+		ICustomizableObjectModule::Get().RegisterExtension(Extension.Get());
+	}
+
+	virtual void ShutdownModule() override
+	{
+		if (ICustomizableObjectModule::IsAvailable()
+			&& Extension.IsValid())
+		{
+			ICustomizableObjectModule::Get().UnregisterExtension(Extension.Get());
+		}
+	}
+
+private:
+	TWeakObjectPtr<const UHairStrandsMutableExtension> Extension;
+};
+
+IMPLEMENT_MODULE(FHairStrandsMutableModule, HairStrandsMutable);
diff --git a/Engine/Plugins/Experimental/HairStrandsMutable/Source/HairStrandsMutable/Public/HairStrandsMutableExtension.h b/Engine/Plugins/Experimental/HairStrandsMutable/Source/HairStrandsMutable/Public/HairStrandsMutableExtension.h
new file mode 100644
--- /dev/null
+++ b/Engine/Plugins/Experimental/HairStrandsMutable/Source/HairStrandsMutable/Public/HairStrandsMutableExtension.h
@@ -0,0 +1,80 @@
+// Copyright Epic Games, Inc. All Rights Reserved.
+
+#pragma once
+
+#include "MuCO/CustomizableObjectExtension.h"
+
+#include "HairStrandsMutableExtension.generated.h"
+
+class UGroomAsset;
+class UGroomCache;
+class UGroomBindingAsset;
+class UPhysicsAsset;
+
+/** Used as ExtensionData to represent a Groom Asset in a Customizable Object graph */
+USTRUCT()
+struct HAIRSTRANDSMUTABLE_API FGroomPinData
+{
+	GENERATED_BODY()
+
+public:
+	/** The index of the mesh component this groom will be attached to */
+	UPROPERTY(EditAnywhere, Category = CustomizableObject)
+	int32 ComponentIndex = 0;
+
+	/** The groom asset to use. See UGroomComponent::GroomAsset. */
+	UPROPERTY(EditAnywhere, Category = Groom)
+	TObjectPtr<UGroomAsset> GroomAsset;
+
+	/** See UGroomComponent::GroomCache. */
+	UPROPERTY(EditAnywhere, Category = GroomCache, meta = (EditCondition = "BindingAsset == nullptr"))
+	TObjectPtr<UGroomCache> GroomCache;
+
+	/** See UGroomComponent::BindingAsset. */
+	UPROPERTY(EditAnywhere, Category = Groom, meta = (EditCondition = "GroomCache == nullptr"))
+	TObjectPtr<UGroomBindingAsset> BindingAsset;
+
+	/** See UGroomComponent::PhysicsAsset. */
+	UPROPERTY(EditAnywhere, Category = Simulation)
+	TObjectPtr<UPhysicsAsset> PhysicsAsset;
+
+	/** See UGroomComponent::AttachmentName. */
+	UPROPERTY(EditAnywhere, Category = Groom)
+	FString AttachmentName;
+
+	/** See UMeshComponent::OverrideMaterials. */
+	UPROPERTY(EditAnywhere, Category = Rendering)
+	TArray<TObjectPtr<class UMaterialInterface>> OverrideMaterials;
+};
+
+USTRUCT()
+struct HAIRSTRANDSMUTABLE_API FGroomInstanceData
+{
+	GENERATED_BODY()
+
+public:
+	UPROPERTY()
+	TArray<FGroomPinData> Grooms;
+};
+
+/**
+ * An extension for Mutable that allows users to bring Grooms from the HairStrands plugin into
+ * their Customizable Objects
+ */
+UCLASS(MinimalAPI)
+class UHairStrandsMutableExtension : public UCustomizableObjectExtension
+{
+	GENERATED_BODY()
+
+public:
+	/** UCustomizableObjectExtension interface */
+	virtual TArray<FCustomizableObjectPinType> GetPinTypes() const override;
+	virtual TArray<FObjectNodeInputPin> GetAdditionalObjectNodePins() const override;
+	virtual FInstancedStruct GenerateExtensionInstanceData(
+		const TArray<FInputPinDataContainer>& InputPinData) const override;
+	virtual void OnCustomizableObjectInstanceUsageUpdated(const UCustomizableObjectInstanceUsage* Usage) const override;
+
+	HAIRSTRANDSMUTABLE_API static const FName GroomPinType;
+	HAIRSTRANDSMUTABLE_API static const FName GroomsBaseNodePinName;
+	HAIRSTRANDSMUTABLE_API static const FText GroomNodeCategory;
+};
diff --git a/Engine/Plugins/Experimental/HairStrandsMutable/Source/HairStrandsMutable/Public/HairStrandsMutableModule.h b/Engine/Plugins/Experimental/HairStrandsMutable/Source/HairStrandsMutable/Public/HairStrandsMutableModule.h
new file mode 100644
--- /dev/null
+++ b/Engine/Plugins/Experimental/HairStrandsMutable/Source/HairStrandsMutable/Public/HairStrandsMutableModule.h
@@ -0,0 +1,7 @@
+// Copyright Epic Games, Inc. All Rights Reserved.
+
+#pragma once
+
+#include "CoreMinimal.h"
+#include "Modules/ModuleInterface.h"
+#include "Modules/ModuleManager.h"
diff --git a/Engine/Plugins/Experimental/HairStrandsMutable/Source/HairStrandsMutableEditor/HairStrandsMutableEditor.Build.cs b/Engine/Plugins/Experimental/HairStrandsMutable/Source/HairStrandsMutableEditor/HairStrandsMutableEditor.Build.cs
new file mode 100644
--- /dev/null
+++ b/Engine/Plugins/Experimental/HairStrandsMutable/Source/HairStrandsMutableEditor/HairStrandsMutableEditor.Build.cs
@@ -0,0 +1,32 @@
+// Copyright Epic Games, Inc. All Rights Reserved.
+
+namespace UnrealBuildTool.Rules
+{
+	public class HairStrandsMutableEditor : ModuleRules
+	{
+        public HairStrandsMutableEditor(ReadOnlyTargetRules Target) : base(Target)
+		{
+			PrivateIncludePaths.Add(ModuleDirectory + "/Private");
+			PublicIncludePaths.Add(ModuleDirectory + "/Public");
+			
+			PublicDependencyModuleNames.AddRange(
+				new string[]
+				{
+					"Core",
+					"CoreUObject",
+					"Engine",
+					"CustomizableObjectEditor",
+					"MutableTools",
+					"MutableRuntime"
+				});
+
+			PrivateDependencyModuleNames.AddRange(
+				new string[]
+				{
+					"HairStrandsMutable",
+					"HairStrandsCore",
+					"StructUtils"
+				});
+		}
+	}
+}
diff --git a/Engine/Plugins/Experimental/HairStrandsMutable/Source/HairStrandsMutableEditor/Private/CustomizableObjectNodeGroomConstant.cpp b/Engine/Plugins/Experimental/HairStrandsMutable/Source/HairStrandsMutableEditor/Private/CustomizableObjectNodeGroomConstant.cpp
new file mode 100644
--- /dev/null
+++ b/Engine/Plugins/Experimental/HairStrandsMutable/Source/HairStrandsMutableEditor/Private/CustomizableObjectNodeGroomConstant.cpp
@@ -0,0 +1,66 @@
+// Copyright Epic Games, Inc. All Rights Reserved.
+
+#include "CustomizableObjectNodeGroomConstant.h"
+
+#include "MuCOE/EdGraphSchema_CustomizableObject.h"
+#include "MuCOE/ExtensionDataCompilerInterface.h"
+#include "MuCOE/RemapPins/CustomizableObjectNodeRemapPins.h"
+#include "MuR/ExtensionData.h"
+#include "MuT/NodeExtensionDataConstant.h"
+#include "HairStrandsMutableExtension.h"
+
+#include UE_INLINE_GENERATED_CPP_BY_NAME(CustomizableObjectNodeGroomConstant)
+
+#define LOCTEXT_NAMESPACE "HairStrandsMutableEditor"
+
+FText UCustomizableObjectNodeGroomConstant::GetNodeTitle(ENodeTitleType::Type TitleType) const
+{
+	return LOCTEXT("Groom_Constant", "Groom Constant");
+}
+
+FLinearColor UCustomizableObjectNodeGroomConstant::GetNodeTitleColor() const
+{
+	const UEdGraphSchema_CustomizableObject* Schema = GetDefault<UEdGraphSchema_CustomizableObject>();
+	return Schema->GetPinTypeColor(UHairStrandsMutableExtension::GroomPinType);
+}
+
+FText UCustomizableObjectNodeGroomConstant::GetTooltipText() const
+{
+	return LOCTEXT("Groom_Constant_Tooltip", "Imports a Groom");
+}
+
+void UCustomizableObjectNodeGroomConstant::AllocateDefaultPins(UCustomizableObjectNodeRemapPins* RemapPins)
+{
+	UEdGraphPin* OutputPin = CustomCreatePin(EGPD_Output, UHairStrandsMutableExtension::GroomPinType, UHairStrandsMutableExtension::GroomsBaseNodePinName);
+	OutputPin->bDefaultValueIsIgnored = true;
+}
+
+bool UCustomizableObjectNodeGroomConstant::ShouldAddToContextMenu(FText& OutCategory) const
+{
+	OutCategory = UHairStrandsMutableExtension::GroomNodeCategory;
+	return true;
+}
+
+mu::Ptr<mu::NodeExtensionData> UCustomizableObjectNodeGroomConstant::GenerateMutableNode(FExtensionDataCompilerInterface& CompilerInterface) const
+{
+	check(IsInGameThread());
+
+	// Create node and extension data container
+	mu::Ptr<mu::NodeExtensionDataConstant> Result = new mu::NodeExtensionDataConstant();
+	
+	// Grooms are usually quite large, so set it up as a streaming constant to allow it to be
+	// loaded on demand.
+	//
+	// If needed we could expose an editable UPROPERTY to give the user the option of making this
+	// an always-loaded constant.
+	UCustomizableObjectResourceDataContainer* Container = nullptr;
+	Result->SetValue(CompilerInterface.MakeStreamedExtensionData(Container));
+	check(Container);
+
+	// Populate instanced struct
+	Container->Data.Data.InitializeAs<FGroomPinData>(GroomData);
+
+	return Result;
+}
+
+#undef LOCTEXT_NAMESPACE
diff --git a/Engine/Plugins/Experimental/HairStrandsMutable/Source/HairStrandsMutableEditor/Private/HairStrandsMutableEditorModule.cpp b/Engine/Plugins/Experimental/HairStrandsMutable/Source/HairStrandsMutableEditor/Private/HairStrandsMutableEditorModule.cpp
new file mode 100644
--- /dev/null
+++ b/Engine/Plugins/Experimental/HairStrandsMutable/Source/HairStrandsMutableEditor/Private/HairStrandsMutableEditorModule.cpp
@@ -0,0 +1,5 @@
+// Copyright Epic Games, Inc. All Rights Reserved.
+
+#include "HairStrandsMutableEditorModule.h"
+
+IMPLEMENT_MODULE(FDefaultModuleImpl, HairStrandsMutableEditor);
diff --git a/Engine/Plugins/Experimental/HairStrandsMutable/Source/HairStrandsMutableEditor/Public/CustomizableObjectNodeGroomConstant.h b/Engine/Plugins/Experimental/HairStrandsMutable/Source/HairStrandsMutableEditor/Public/CustomizableObjectNodeGroomConstant.h
new file mode 100644
--- /dev/null
+++ b/Engine/Plugins/Experimental/HairStrandsMutable/Source/HairStrandsMutableEditor/Public/CustomizableObjectNodeGroomConstant.h
@@ -0,0 +1,33 @@
+// Copyright Epic Games, Inc. All Rights Reserved.
+
+#pragma once
+
+#include "HairStrandsMutableExtension.h"
+#include "MuCOE/ICustomizableObjectExtensionNode.h"
+#include "MuCOE/Nodes/CustomizableObjectNodeExtensionDataConstant.h"
+
+#include "CustomizableObjectNodeGroomConstant.generated.h"
+
+/** Imports a Groom into the Customizable Object graph */
+UCLASS()
+class UCustomizableObjectNodeGroomConstant : public UCustomizableObjectNodeExtensionDataConstant, public ICustomizableObjectExtensionNode
+{
+	GENERATED_BODY()
+
+public:
+	/** EdGraphNode interface */
+	virtual FText GetNodeTitle(ENodeTitleType::Type TitleType) const override;
+	virtual FLinearColor GetNodeTitleColor() const override;
+	virtual FText GetTooltipText() const override;
+
+	/** UCustomizableObjectNode interface */
+	virtual void AllocateDefaultPins(UCustomizableObjectNodeRemapPins* RemapPins) override;
+	virtual bool ShouldAddToContextMenu(FText& OutCategory) const override;
+
+	/** ICustomizableObjectExtensionNode interface */
+	virtual mu::Ptr<mu::NodeExtensionData> GenerateMutableNode(FExtensionDataCompilerInterface& CompilerInterface) const override;
+
+	UPROPERTY(EditAnywhere, Category = CustomizableObject)
+	FGroomPinData GroomData;
+};
+
diff --git a/Engine/Plugins/Experimental/HairStrandsMutable/Source/HairStrandsMutableEditor/Public/HairStrandsMutableEditorModule.h b/Engine/Plugins/Experimental/HairStrandsMutable/Source/HairStrandsMutableEditor/Public/HairStrandsMutableEditorModule.h
new file mode 100644
--- /dev/null
+++ b/Engine/Plugins/Experimental/HairStrandsMutable/Source/HairStrandsMutableEditor/Public/HairStrandsMutableEditorModule.h
@@ -0,0 +1,7 @@
+// Copyright Epic Games, Inc. All Rights Reserved.
+
+#pragma once
+
+#include "CoreMinimal.h"
+#include "Modules/ModuleInterface.h"
+#include "Modules/ModuleManager.h"
