[HairStrandsMutable] Marked the plugin as experimental.

Was marked as beta.

#rnx

diff --git a/Engine/Plugins/Experimental/HairStrandsMutable/HairStrandsMutable.uplugin b/Engine/Plugins/Experimental/HairStrandsMutable/HairStrandsMutable.uplugin
--- a/Engine/Plugins/Experimental/HairStrandsMutable/HairStrandsMutable.uplugin
+++ b/Engine/Plugins/Experimental/HairStrandsMutable/HairStrandsMutable.uplugin
@@ -11,7 +11,7 @@
 	"MarketplaceURL": "",
 	"SupportURL": "",
 	"CanContainContent": false,
-	"IsBetaVersion": true,
+	"IsBetaVersion": false,
 	"IsExperimentalVersion": true,
 	"Installed": false,
 	"EnabledByDefault": false,
