[HairStrandsMutableExtension] Switch to the new Mutable ComponentName

#rnx

diff --git a/Engine/Plugins/Experimental/HairStrandsMutable/Source/HairStrandsMutable/Private/HairStrandsMutableExtension.cpp b/Engine/Plugins/Experimental/HairStrandsMutable/Source/HairStrandsMutable/Private/HairStrandsMutableExtension.cpp
--- a/Engine/Plugins/Experimental/HairStrandsMutable/Source/HairStrandsMutable/Private/HairStrandsMutableExtension.cpp
+++ b/Engine/Plugins/Experimental/HairStrandsMutable/Source/HairStrandsMutable/Private/HairStrandsMutableExtension.cpp
@@ -93,7 +93,7 @@ void UHairStrandsMutableExtension::OnCustomizableObjectInstanceUsageUpdated(cons
 
 	for (const FGroomPinData& GroomPinData : InstanceData->Grooms)
 	{
-		if (GroomPinData.ComponentIndex != InstanceUsage->GetComponentIndex())
+		if (GroomPinData.ComponentName != InstanceUsage->GetComponentName())
 		{
 			continue;
 		}
diff --git a/Engine/Plugins/Experimental/HairStrandsMutable/Source/HairStrandsMutable/Public/HairStrandsMutableExtension.h b/Engine/Plugins/Experimental/HairStrandsMutable/Source/HairStrandsMutable/Public/HairStrandsMutableExtension.h
--- a/Engine/Plugins/Experimental/HairStrandsMutable/Source/HairStrandsMutable/Public/HairStrandsMutableExtension.h
+++ b/Engine/Plugins/Experimental/HairStrandsMutable/Source/HairStrandsMutable/Public/HairStrandsMutableExtension.h
@@ -18,9 +18,9 @@ struct HAIRSTRANDSMUTABLE_API FGroomPinData
 	GENERATED_BODY()
 
 public:
-	/** The index of the mesh component this groom will be attached to */
+	/** Name of the mesh component this groom will be attached to */
 	UPROPERTY(EditAnywhere, Category = CustomizableObject)
-	int32 ComponentIndex = 0;
+	FName ComponentName;
 
 	/** The groom asset to use. See UGroomComponent::GroomAsset. */
 	UPROPERTY(EditAnywhere, Category = Groom)
