[HairStrandsMutable] Removed StructUtils dependency.

#rnx

diff --git a/Engine/Plugins/Experimental/HairStrandsMutable/HairStrandsMutable.uplugin b/Engine/Plugins/Experimental/HairStrandsMutable/HairStrandsMutable.uplugin
--- a/Engine/Plugins/Experimental/HairStrandsMutable/HairStrandsMutable.uplugin
+++ b/Engine/Plugins/Experimental/HairStrandsMutable/HairStrandsMutable.uplugin
@@ -36,10 +36,6 @@
 		{
 			"Name": "Mutable",
 			"Enabled": true
-		},
-		{
-			"Name": "StructUtils",
-			"Enabled": true	
 		}
 	]
 
diff --git a/Engine/Plugins/Experimental/HairStrandsMutable/Source/HairStrandsMutable/HairStrandsMutable.Build.cs b/Engine/Plugins/Experimental/HairStrandsMutable/Source/HairStrandsMutable/HairStrandsMutable.Build.cs
--- a/Engine/Plugins/Experimental/HairStrandsMutable/Source/HairStrandsMutable/HairStrandsMutable.Build.cs
+++ b/Engine/Plugins/Experimental/HairStrandsMutable/Source/HairStrandsMutable/HairStrandsMutable.Build.cs
@@ -17,7 +17,6 @@ namespace UnrealBuildTool.Rules
 					"Engine",
 					"CustomizableObject",
 					"HairStrandsCore",
-					"StructUtils",
 				});
 
 			PrivateDependencyModuleNames.AddRange(
diff --git a/Engine/Plugins/Experimental/HairStrandsMutable/Source/HairStrandsMutableEditor/HairStrandsMutableEditor.Build.cs b/Engine/Plugins/Experimental/HairStrandsMutable/Source/HairStrandsMutableEditor/HairStrandsMutableEditor.Build.cs
--- a/Engine/Plugins/Experimental/HairStrandsMutable/Source/HairStrandsMutableEditor/HairStrandsMutableEditor.Build.cs
+++ b/Engine/Plugins/Experimental/HairStrandsMutable/Source/HairStrandsMutableEditor/HairStrandsMutableEditor.Build.cs
@@ -25,7 +25,6 @@ namespace UnrealBuildTool.Rules
 				{
 					"HairStrandsMutable",
 					"HairStrandsCore",
-					"StructUtils"
 				});
 		}
 	}
