[Mutable] Add component naming support for spawned groom components

#13307

diff --git a/Engine/Plugins/Experimental/HairStrandsMutable/Source/HairStrandsMutable/Private/HairStrandsMutableExtension.cpp b/Engine/Plugins/Experimental/HairStrandsMutable/Source/HairStrandsMutable/Private/HairStrandsMutableExtension.cpp
--- a/Engine/Plugins/Experimental/HairStrandsMutable/Source/HairStrandsMutable/Private/HairStrandsMutableExtension.cpp
+++ b/Engine/Plugins/Experimental/HairStrandsMutable/Source/HairStrandsMutable/Private/HairStrandsMutableExtension.cpp
@@ -128,7 +128,7 @@ void UHairStrandsMutableExtension::OnCustomizableObjectInstanceUsageUpdated(UCus
 			continue;
 		}
 		
-		UGroomComponent* GroomComponent = NewObject<UGroomComponent>(AttachParent);
+		UGroomComponent* GroomComponent = NewObject<UGroomComponent>(AttachParent, GroomPinData.GroomComponentName);
 		GroomComponent->GroomAsset = GroomPinData.GroomAsset;
 		GroomComponent->GroomCache = GroomPinData.GroomCache;
 		GroomComponent->BindingAsset = GroomPinData.BindingAsset;
diff --git a/Engine/Plugins/Experimental/HairStrandsMutable/Source/HairStrandsMutable/Public/HairStrandsMutableExtension.h b/Engine/Plugins/Experimental/HairStrandsMutable/Source/HairStrandsMutable/Public/HairStrandsMutableExtension.h
--- a/Engine/Plugins/Experimental/HairStrandsMutable/Source/HairStrandsMutable/Public/HairStrandsMutableExtension.h
+++ b/Engine/Plugins/Experimental/HairStrandsMutable/Source/HairStrandsMutable/Public/HairStrandsMutableExtension.h
@@ -10,6 +10,7 @@ class UGroomAsset;
 class UGroomCache;
 class UGroomBindingAsset;
 class UPhysicsAsset;
+class UMaterialInterface;
 
 /** Used as ExtensionData to represent a Groom Asset in a Customizable Object graph */
 USTRUCT()
@@ -42,9 +43,13 @@ public:
 	UPROPERTY(EditAnywhere, Category = Groom)
 	FString AttachmentName;
 
+	/** Name of the created groom component */
+	UPROPERTY(EditAnywhere, Category = Groom)
+	FName GroomComponentName;
+
 	/** See UMeshComponent::OverrideMaterials. */
 	UPROPERTY(EditAnywhere, Category = Rendering)
-	TArray<TObjectPtr<class UMaterialInterface>> OverrideMaterials;
+	TArray<TObjectPtr<UMaterialInterface>> OverrideMaterials;
 };
 
 USTRUCT()
