[mutable] Changed friendly name to the MutableDataflow and HairStrandsMutable experimental plugins.
- Fixed incorrect MutableDataflow comment.

#rnx

diff --git a/Engine/Plugins/Experimental/HairStrandsMutable/HairStrandsMutable.uplugin b/Engine/Plugins/Experimental/HairStrandsMutable/HairStrandsMutable.uplugin
--- a/Engine/Plugins/Experimental/HairStrandsMutable/HairStrandsMutable.uplugin
+++ b/Engine/Plugins/Experimental/HairStrandsMutable/HairStrandsMutable.uplugin
@@ -2,7 +2,7 @@
 	"FileVersion": 3,
 	"Version": 1,
 	"VersionName": "0.1",
-	"FriendlyName": "Groom Extensions For Mutable",
+	"FriendlyName": "Mutable Groom Extensions",
 	"Description": "Adds Mutable functionality to work with Grooms from the HairStrands plugin",
 	"Category": "Animation",
 	"CreatedBy": "Epic Games, Inc.",
diff --git a/Engine/Plugins/Experimental/MutableDataflow/MutableDataflow.uplugin b/Engine/Plugins/Experimental/MutableDataflow/MutableDataflow.uplugin
--- a/Engine/Plugins/Experimental/MutableDataflow/MutableDataflow.uplugin
+++ b/Engine/Plugins/Experimental/MutableDataflow/MutableDataflow.uplugin
@@ -2,7 +2,7 @@
 	"FileVersion": 3,
 	"Version": 1,
 	"VersionName": "0.1",
-	"FriendlyName": "Dataflow Extensions For Mutable",
+	"FriendlyName": "Mutable Dataflow Extensions",
 	"Description": "Adds Mutable functionality to work with Dataflow objects from the Dataflow plugin",
 	"Category": "Animation",
 	"CreatedBy": "Epic Games, Inc.",
