[Mutable] Fix PSO check with grooms.

#Synced-CL 50250210

diff --git a/Engine/Plugins/Experimental/HairStrandsMutable/Source/HairStrandsMutable/Private/HairStrandsMutableExtension.cpp b/Engine/Plugins/Experimental/HairStrandsMutable/Source/HairStrandsMutable/Private/HairStrandsMutableExtension.cpp
--- a/Engine/Plugins/Experimental/HairStrandsMutable/Source/HairStrandsMutable/Private/HairStrandsMutableExtension.cpp
+++ b/Engine/Plugins/Experimental/HairStrandsMutable/Source/HairStrandsMutable/Private/HairStrandsMutableExtension.cpp
@@ -142,6 +142,7 @@ void UHairStrandsMutableExtension::OnCustomizableObjectInstanceUsageUpdated(UCus
 		GroomComponent->ComponentTags.Add(GroomComponentTag);
 		
 		GroomComponent->AttachToComponent(AttachParent, FAttachmentTransformRules::KeepRelativeTransform);
+		GroomComponent->PrecachePSOs();
 		GroomComponent->RegisterComponent();
 
 		UpdateGroomComponents.Add(GroomComponent);
