[HairStrandsMutable] Fix ensure when registgering to Movie Sequence Template actors.

#Synced-CL 51579983

diff --git a/Engine/Plugins/Experimental/HairStrandsMutable/Source/HairStrandsMutable/Private/HairStrandsMutableExtension.cpp b/Engine/Plugins/Experimental/HairStrandsMutable/Source/HairStrandsMutable/Private/HairStrandsMutableExtension.cpp
--- a/Engine/Plugins/Experimental/HairStrandsMutable/Source/HairStrandsMutable/Private/HairStrandsMutableExtension.cpp
+++ b/Engine/Plugins/Experimental/HairStrandsMutable/Source/HairStrandsMutable/Private/HairStrandsMutableExtension.cpp
@@ -119,7 +119,13 @@ void UHairStrandsMutableExtension::OnCustomizableObjectInstanceUsageUpdated(UCus
 		
 		GroomComponent->AttachToComponent(AttachParent, FAttachmentTransformRules::KeepRelativeTransform);
 		GroomComponent->PrecachePSOs();
-		GroomComponent->RegisterComponent();
+	
+		AActor* MyOwner = GroomComponent->GetOwner();
+		UWorld* MyOwnerWorld = (MyOwner ? MyOwner->GetWorld() : nullptr);
+		if (MyOwnerWorld) // Avoid ensure when registering on Movie Sequence Template actors.
+		{
+			GroomComponent->RegisterComponent();
+		}
 
 		UpdateGroomComponents.Add(GroomComponent);
 	}
